fixup the plugin name check for GDK DLC to decide whether to show the configuration or not.

#rnx
#Synced-CL 51451766

diff --git a/Engine/Source/Developer/Microsoft/GDKPlatformEditor/Private/GDKTargetSettingsDetails.cpp b/Engine/Source/Developer/Microsoft/GDKPlatformEditor/Private/GDKTargetSettingsDetails.cpp
--- a/Engine/Source/Developer/Microsoft/GDKPlatformEditor/Private/GDKTargetSettingsDetails.cpp
+++ b/Engine/Source/Developer/Microsoft/GDKPlatformEditor/Private/GDKTargetSettingsDetails.cpp
@@ -172,7 +172,7 @@ void FGDKTargetSettingsDetails::CustomizeDetails(IDetailLayoutBuilder& DetailBui
 	];
 
 	// hide the DLC category if the GDK DLC support plugin isn't enabled
-	TSharedPtr<IPlugin> GDKDLCPlugin = IPluginManager::Get().FindPlugin(TEXT("GDKDLCSupport"));
+	TSharedPtr<IPlugin> GDKDLCPlugin = IPluginManager::Get().FindPlugin(TEXT("GDKPlatformDLC"));
 	if (!GDKDLCPlugin.IsValid() || !GDKDLCPlugin->IsEnabled())
 	{
 		DetailBuilder.HideCategory("DLC");
